PR #5395: UE-65844: Make string table searchable (Contributed by [REDACTED_PERSON])

diff --git a/Engine/Source/Editor/StringTableEditor/Private/StringTableEditor.cpp b/Engine/Source/Editor/StringTableEditor/Private/StringTableEditor.cpp
--- a/Engine/Source/Editor/StringTableEditor/Private/StringTableEditor.cpp
+++ b/Engine/Source/Editor/StringTableEditor/Private/StringTableEditor.cpp
@@ -13,12 +13,17 @@
 #include "Editor.h"
 #include "ScopedTransaction.h"
 #include "Framework/Application/SlateApplication.h"
+#include "Framework/Commands/UICommandList.h"
+#include "Framework/Commands/GenericCommands.h"
+#include "Framework/MultiBox/MultiBoxBuilder.h"
 
 #include "Widgets/Input/SButton.h"
 #include "Widgets/Images/SImage.h"
 #include "Widgets/Docking/SDockTab.h"
 #include "Widgets/Text/STextBlock.h"
 #include "Widgets/Views/SHeaderRow.h"
+#include "Widgets/Input/SSearchBox.h"
+#include "Widgets/Text/SInlineEditableTextBlock.h"
  
 #define LOCTEXT_NAMESPACE "StringTableEditor"
 
@@ -50,9 +55,13 @@ public:
 		}
 		else if (ColumnName == FStringTableEditor::StringTableKeyColumnId)
 		{
-			Return = SNew(SEditableTextBox)
-				.IsReadOnly(true)
-				.Text(FText::FromString(CachedStringTableEntry->Key));
+			Return = SNew(SBox)
+				.Padding(4)
+				[
+					SNew(SInlineEditableTextBlock)
+					.IsReadOnly(true)
+					.Text(FText::FromString(CachedStringTableEntry->Key))
+				];
 		}
 		else if (ColumnName == FStringTableEditor::StringTableSourceStringColumnId)
 		{
@@ -122,10 +131,6 @@ void FStringTableEditor::UnregisterTabSpawners(const TSharedRef<class FTabManage
 	InTabManager->UnregisterTabSpawner(StringTableTabId);
 }
 
-FStringTableEditor::FStringTableEditor()
-{
-}
-
 FStringTableEditor::~FStringTableEditor()
 {
 	GEditor->UnregisterForUndo(this);
@@ -228,6 +233,13 @@ TSharedRef<SDockTab> FStringTableEditor::SpawnTab_StringTable(const FSpawnTabArg
 {
 	check(Args.GetTabId().TabType == StringTableTabId);
 
+	const FGenericCommands& GenericCommands = FGenericCommands::Get();
+	CommandList = MakeShareable(new FUICommandList);
+	CommandList->MapAction(
+		GenericCommands.Delete,
+		FExecuteAction::CreateSP(this, &FStringTableEditor::ExecuteDelete),
+		FCanExecuteAction::CreateSP(this, &FStringTableEditor::CanExecuteDelete));
+
 	UStringTable* StringTable = Cast<UStringTable>(GetEditingObject());
 
 	// Support undo/redo
@@ -236,19 +248,24 @@ TSharedRef<SDockTab> FStringTableEditor::SpawnTab_StringTable(const FSpawnTabArg
 		StringTable->SetFlags(RF_Transactional);
 	}
 
+	EntryTextFilter = MakeShareable(new FEntryTextFilter(FEntryTextFilter::FItemToStringArray::CreateLambda([](const TSharedPtr<FCachedStringTableEntry>& InItem, OUT TArray< FString >& StringArray) {
+		StringArray.Add(InItem->Key);
+		StringArray.Add(InItem->SourceString);
+	})));
+
 	TSharedRef<SDockTab> StringTableTab = SNew(SDockTab)
 		.Icon(FEditorStyle::GetBrush("StringTableEditor.Tabs.Properties"))
 		.Label(LOCTEXT("StringTableTitle", "String Table"))
 		.TabColorScale(GetTabColorScale())
 		[
-			SNew(SBorder)
-			.Padding(2)
-			.BorderImage(FEditorStyle::GetBrush("ToolPanel.GroupBorder"))
-			[
-				SNew(SVerticalBox)
+			SNew(SVerticalBox)
 
-				+SVerticalBox::Slot()
-				.AutoHeight()
+			+SVerticalBox::Slot()
+			.AutoHeight()
+			[
+				SNew(SBorder)
+				.Padding(2)
+				.BorderImage(FEditorStyle::GetBrush("ToolPanel.GroupBorder"))
 				[
 					SNew(SHorizontalBox)
 
@@ -279,6 +296,9 @@ TSharedRef<SDockTab> FStringTableEditor::SpawnTab_StringTable(const FSpawnTabArg
 						SNew(SButton)
 						.Text(LOCTEXT("ImportFromCSVLabel", "Import from CSV"))
 						.OnClicked(this, &FStringTableEditor::OnImportFromCSVClicked)
+						.ButtonStyle(FEditorStyle::Get(), "FlatButton.Default")
+						.ForegroundColor(FLinearColor::White)
+						.ContentPadding(FMargin(6, 2))
 					]
 
 					+SHorizontalBox::Slot()
@@ -289,15 +309,37 @@ TSharedRef<SDockTab> FStringTableEditor::SpawnTab_StringTable(const FSpawnTabArg
 						SNew(SButton)
 						.Text(LOCTEXT("ExportToCSVLabel", "Export to CSV"))
 						.OnClicked(this, &FStringTableEditor::OnExportToCSVClicked)
+						.ButtonStyle(FEditorStyle::Get(), "FlatButton.Default")
+						.ForegroundColor(FLinearColor::White)
+						.ContentPadding(FMargin(6, 2))
 					]
 				]
+			]
 
-				+SVerticalBox::Slot()
+			+ SVerticalBox::Slot()
+			.AutoHeight()
+			.Padding(0, 4, 0, 4)
+			[
+				SNew(SBorder)
+				.Padding(2)
+				.BorderImage(FEditorStyle::GetBrush("ToolPanel.GroupBorder"))
+				[
+					SAssignNew(SearchBox, SSearchBox)
+					.OnTextChanged(this, &FStringTableEditor::OnFilterTextChanged)
+				]
+			]
+
+			+SVerticalBox::Slot()
+			[
+				SNew(SBorder)
+				.Padding(2)
+				.BorderImage(FEditorStyle::GetBrush("ToolPanel.GroupBorder"))
 				[
 					SAssignNew(StringTableEntriesListView, SListView<TSharedPtr<FCachedStringTableEntry>>)
 					.ListItemsSource(&CachedStringTableEntries)
 					.OnGenerateRow(this, &FStringTableEditor::OnGenerateStringTableEntryRow)
-					.SelectionMode(ESelectionMode::Single)
+					.SelectionMode(ESelectionMode::Multi)
+					.OnContextMenuOpening(this, &FStringTableEditor::OnConstructContextMenu)
 					.HeaderRow
 					(
 						SNew(SHeaderRow)
@@ -319,9 +361,15 @@ TSharedRef<SDockTab> FStringTableEditor::SpawnTab_StringTable(const FSpawnTabArg
 						.FixedWidth(28)
 					)
 				]
+			]
 
-				+SVerticalBox::Slot()
-				.AutoHeight()
+			+SVerticalBox::Slot()
+			.Padding(0, 4, 0, 0)
+			.AutoHeight()
+			[
+				SNew(SBorder)
+				.Padding(2)
+				.BorderImage(FEditorStyle::GetBrush("ToolPanel.GroupBorder"))
 				[
 					SNew(SHorizontalBox)
 
@@ -369,6 +417,9 @@ TSharedRef<SDockTab> FStringTableEditor::SpawnTab_StringTable(const FSpawnTabArg
 						SNew(SButton)
 						.Text(LOCTEXT("AddLabel", "Add"))
 						.OnClicked(this, &FStringTableEditor::OnAddClicked)
+						.ButtonStyle(FEditorStyle::Get(), "FlatButton.Success")
+						.ForegroundColor(FLinearColor::White)
+						.ContentPadding(FMargin(6, 2))
 					]
 				]
 			]
@@ -379,6 +430,20 @@ TSharedRef<SDockTab> FStringTableEditor::SpawnTab_StringTable(const FSpawnTabArg
 	return StringTableTab;
 }
 
+TSharedPtr<SWidget> FStringTableEditor::OnConstructContextMenu()
+{
+	const bool bShouldCloseWindowAfterMenuSelection = true;
+	FMenuBuilder MenuBuilder(bShouldCloseWindowAfterMenuSelection, CommandList);
+
+	MenuBuilder.BeginSection("EditOperations", LOCTEXT("StringTableEditOperationsLabel", "Edit"));
+	{
+		MenuBuilder.AddMenuEntry(FGenericCommands::Get().Delete);
+	}
+	MenuBuilder.EndSection();
+
+	return MenuBuilder.MakeWidget();
+}
+
 TSharedRef<ITableRow> FStringTableEditor::OnGenerateStringTableEntryRow(TSharedPtr<FCachedStringTableEntry> CachedStringTableEntry, const TSharedRef<STableViewBase>& Table)
 {
 	return SNew(SStringTableEntryRow, Table, *this, CachedStringTableEntry.ToSharedRef());
@@ -396,7 +461,11 @@ void FStringTableEditor::RefreshCachedStringTable(const FString& InCachedSelecti
 		StringTable->GetStringTable()->EnumerateSourceStrings([&](const FString& InKey, const FString& InSourceString)
 		{
 			TSharedRef<FCachedStringTableEntry> NewStringTableEntry = MakeShared<FCachedStringTableEntry>(InKey, InSourceString);
-			CachedStringTableEntries.Add(NewStringTableEntry);
+
+			if (EntryTextFilter->PassesFilter(NewStringTableEntry))
+			{
+				CachedStringTableEntries.Add(NewStringTableEntry);
+			}
 
 			if (InCachedSelection.Equals(InKey, ESearchCase::CaseSensitive))
 			{
@@ -532,6 +601,30 @@ void FStringTableEditor::OnSourceStringCommitted(const FText& InText, ETextCommi
 	}
 }
 
+void FStringTableEditor::OnFilterTextChanged(const FText& InNewText)
+{
+	EntryTextFilter->SetRawFilterText(InNewText);
+	SearchBox->SetError(EntryTextFilter->GetFilterErrorText());
+
+	RefreshCachedStringTable();
+
+	StringTableEntriesListView->RebuildList();
+}
+
+bool FStringTableEditor::CanExecuteDelete() const
+{
+	return StringTableEntriesListView->GetNumItemsSelected() > 0;
+}
+
+void FStringTableEditor::ExecuteDelete()
+{
+	const TArray<TSharedPtr<FCachedStringTableEntry>>& Items = StringTableEntriesListView->GetSelectedItems();
+	for (const auto& Item : Items)
+	{
+		DeleteEntry(Item->Key);
+	}
+}
+
 FReply FStringTableEditor::OnAddClicked()
 {
 	const FText NewKey = KeyEditableTextBox->GetText();
diff --git a/Engine/Source/Editor/StringTableEditor/Private/StringTableEditor.h b/Engine/Source/Editor/StringTableEditor/Private/StringTableEditor.h
--- a/Engine/Source/Editor/StringTableEditor/Private/StringTableEditor.h
+++ b/Engine/Source/Editor/StringTableEditor/Private/StringTableEditor.h
@@ -10,6 +10,7 @@
 #include "Widgets/Input/SEditableTextBox.h"
 #include "Widgets/Input/SMultiLineEditableTextBox.h"
 #include "Widgets/Views/SListView.h"
+#include "Misc/TextFilter.h"
 
 /** Viewer/editor for a String Table */
 class FStringTableEditor : public IStringTableEditor, public FEditorUndoClient
@@ -30,7 +31,7 @@ public:
 	void InitStringTableEditor(const EToolkitMode::Type Mode, const TSharedPtr<class IToolkitHost>& InitToolkitHost, UStringTable* StringTable);
 
 	/** Constructor */
-	FStringTableEditor();
+	FStringTableEditor() = default;
 
 	/** Destructor */
 	virtual ~FStringTableEditor();
@@ -56,10 +57,6 @@ private:
 	/** Cached string table entry */
 	struct FCachedStringTableEntry
 	{
-		FCachedStringTableEntry()
-		{
-		}
-
 		FCachedStringTableEntry(FString InKey, FString InSourceString)
 			: Key(MoveTemp(InKey))
 			, SourceString(MoveTemp(InSourceString))
@@ -103,6 +100,18 @@ private:
 	/** Handler for the source string being committed. Treat as if "Add" were pressed. */
 	void OnSourceStringCommitted(const FText& InText, ETextCommit::Type);
 
+	/** Handler for the search box string being changed. */
+	void OnFilterTextChanged(const FText& InNewText);
+
+	/** Check function for context menu delete action. */
+	bool CanExecuteDelete() const;
+
+	/** Handler for context menu delete action. */
+	void ExecuteDelete();
+
+	/** Handler for when opening the context menu in the list view. */
+	TSharedPtr<SWidget> OnConstructContextMenu();
+
 	/** Handler for the "Add" button */
 	FReply OnAddClicked();
 
@@ -115,6 +124,13 @@ private:
 	/** Editable text for the namespace */
 	TSharedPtr<SEditableTextBox> NamespaceEditableTextBox;
 
+	/** Search box for the text entries */
+	TSharedPtr<class SSearchBox> SearchBox;
+
+	/** Text filter for text entries */
+	using FEntryTextFilter = TTextFilter< TSharedPtr<FCachedStringTableEntry> >;
+	TSharedPtr<FEntryTextFilter> EntryTextFilter;
+
 	/** Editable text for the key */
 	TSharedPtr<SEditableTextBox> KeyEditableTextBox;
 
@@ -127,6 +143,9 @@ private:
 	/** List view showing the cached string table entries */
 	TSharedPtr<SListView<TSharedPtr<FCachedStringTableEntry>>> StringTableEntriesListView;
 
+	/** List of UI commands */
+	TSharedPtr<class FUICommandList> CommandList;
+
 	/**	The tab ID for the string table tab */
 	static const FName StringTableTabId;
 
