PR #8121: Updated string table CSV import to filter only CSV extension (Contributed by [REDACTED_PERSON])

diff --git a/Engine/Source/Editor/StringTableEditor/Private/StringTableEditor.cpp b/Engine/Source/Editor/StringTableEditor/Private/StringTableEditor.cpp
--- a/Engine/Source/Editor/StringTableEditor/Private/StringTableEditor.cpp
+++ b/Engine/Source/Editor/StringTableEditor/Private/StringTableEditor.cpp
@@ -662,7 +662,7 @@ FReply FStringTableEditor::OnImportFromCSVClicked()
 			LOCTEXT("ImportStringTableTitle", "Choose a string table CSV file...").ToString(),
 			DefaultPath,
 			TEXT(""),
-			TEXT("String Table CSV (*.csv)"),
+			TEXT("String Table CSV (*.csv)|*.csv"),
 			EFileDialogFlags::None,
 			OutFiles
 			))
